Add OpenXR Vive Tracker extension plugin

diff --git a/Engine/Plugins/Runtime/OpenXRViveTracker/OpenXRViveTracker.uplugin b/Engine/Plugins/Runtime/OpenXRViveTracker/OpenXRViveTracker.uplugin
new file mode 100644
--- /dev/null
+++ b/Engine/Plugins/Runtime/OpenXRViveTracker/OpenXRViveTracker.uplugin
@@ -0,0 +1,45 @@
+{
+	"FileVersion": 3,
+	"Version": 1,
+	"VersionName": "1.0",
+	"FriendlyName": "OpenXRViveTracker",
+	"Description": "OpenXR Vive Tracker provides XR_HTCX_vive_tracker_interaction.",
+	"Category": "Virtual Reality",
+	"CreatedBy": "Epic Games, Inc.",
+	"CreatedByURL": "http://epicgames.com",
+	"DocsURL": "",
+	"MarketplaceURL": "",
+	"SupportURL": "",
+	"EnabledByDefault": false,
+	"CanContainContent": true,
+	"IsBetaVersion": true,
+	"Installed": false,
+	"SupportedTargetPlatforms": [
+		"Win64",
+		"Linux",
+		"Android"
+	],
+	"Modules": [
+		{
+			"Name": "OpenXRViveTracker",
+			"Type": "Runtime",
+			"LoadingPhase": "PostConfigInit",
+			"PlatformAllowList": [
+				"Win64",
+				"Linux",
+				"Android"
+			]
+		}
+	],
+	"Plugins": [
+		{
+			"Name": "OpenXR",
+			"Enabled": true,
+			"PlatformAllowList": [
+				"Win64",
+				"Linux",
+				"Android"
+			]
+		}
+	]
+}
\ No newline at end of file
diff --git a/Engine/Plugins/Runtime/OpenXRViveTracker/Source/OpenXRViveTracker/OpenXRViveTracker.Build.cs b/Engine/Plugins/Runtime/OpenXRViveTracker/Source/OpenXRViveTracker/OpenXRViveTracker.Build.cs
new file mode 100644
--- /dev/null
+++ b/Engine/Plugins/Runtime/OpenXRViveTracker/Source/OpenXRViveTracker/OpenXRViveTracker.Build.cs
@@ -0,0 +1,52 @@
+// Copyright Epic Games, Inc. All Rights Reserved.
+
+using UnrealBuildTool;
+using System.IO;
+
+namespace UnrealBuildTool.Rules
+{
+    public class OpenXRViveTracker: ModuleRules
+    {
+        public OpenXRViveTracker(ReadOnlyTargetRules Target) 
+				: base(Target)
+        {
+			PublicDependencyModuleNames.AddRange(
+			   new string[]
+			   {
+					"InputDevice"
+			   }
+		   );
+
+			var EngineDir = Path.GetFullPath(Target.RelativeEnginePath);
+            PrivateIncludePaths.AddRange(
+                new string[] {
+                    EngineDir + "/Source/ThirdParty/OpenXR/include",
+					EngineDir + "/Plugins/Runtime/OpenXR/Source/OpenXRHMD/Private",
+				}
+                );
+
+            PrivateDependencyModuleNames.AddRange(
+                new string[]
+                {
+                    "Core",
+                    "CoreUObject",
+                    "Engine",
+					"HeadMountedDisplay",
+                    "InputCore",
+					"OpenXRHMD",
+					"OpenXRInput",
+					"Slate",
+					"SlateCore",
+					"ApplicationCore"
+				}
+				);
+
+            AddEngineThirdPartyPrivateStaticDependencies(Target, "OpenXR");
+
+            if (Target.bBuildEditor == true)
+            {
+                PrivateDependencyModuleNames.Add("UnrealEd");
+            }
+        }
+    }
+}
diff --git a/Engine/Plugins/Runtime/OpenXRViveTracker/Source/OpenXRViveTracker/Private/OpenXRViveTracker.cpp b/Engine/Plugins/Runtime/OpenXRViveTracker/Source/OpenXRViveTracker/Private/OpenXRViveTracker.cpp
new file mode 100644
--- /dev/null
+++ b/Engine/Plugins/Runtime/OpenXRViveTracker/Source/OpenXRViveTracker/Private/OpenXRViveTracker.cpp
@@ -0,0 +1,486 @@
+// Copyright Epic Games, Inc. All Rights Reserved.
+
+#include "OpenXRViveTracker.h"
+#include "IOpenXRViveTrackerModule.h"
+
+#include "IXRTrackingSystem.h"
+#include "CoreMinimal.h"
+#include "UObject/Package.h"
+#include "UObject/UObjectGlobals.h"
+#include "UObject/ObjectMacros.h"
+#include "Engine/Engine.h"
+#include "OpenXRCore.h"
+#include "Framework/Application/SlateApplication.h"
+#include "IOpenXRExtensionPlugin.h"
+#include "Modules/ModuleManager.h"
+#include "Features/IModularFeatures.h"
+#include "HeadMountedDisplayFunctionLibrary.h"
+#include "IOpenXRHMDModule.h"
+
+#define LOCTEXT_NAMESPACE "OpenXRViveTracker"
+
+DEFINE_LOG_CATEGORY_STATIC(LogOpenXRViveTracker, Display, All);
+
+class FOpenXRViveTrackerModule :
+	public IOpenXRViveTrackerModule,
+	public IOpenXRExtensionPlugin
+{
+public:
+	FOpenXRViveTrackerModule()
+		: InputDevice(nullptr)
+	{}
+
+	virtual void StartupModule() override
+	{
+		IOpenXRViveTrackerModule::StartupModule();
+		RegisterOpenXRExtensionModularFeature();
+
+		// HACK: Generic Application might not be instantiated at this point so we create the input device with a
+		// dummy message handler. When the Generic Application creates the input device it passes a valid message
+		// handler to it which is further on used for all the controller events. This hack fixes issues caused by
+		// using a custom input device before the Generic Application has instantiated it. Eg. within BeginPlay()
+		//
+		// This also fixes the warnings that pop up on the custom input keys when the blueprint loads. Those
+		// warnings are caused because Unreal loads the bluerints before the input device has been instantiated
+		// and has added its keys, thus leading Unreal to believe that those keys don't exist. This hack causes
+		// an earlier instantiation of the input device, and consequently, the custom keys.
+		TSharedPtr<FGenericApplicationMessageHandler> DummyMessageHandler(new FGenericApplicationMessageHandler());
+		CreateInputDevice(DummyMessageHandler.ToSharedRef());
+	}
+
+	virtual void ShutdownModule() override
+	{
+		IOpenXRViveTrackerModule::ShutdownModule();
+	}
+
+	virtual TSharedPtr<class IInputDevice> CreateInputDevice(const TSharedRef<FGenericApplicationMessageHandler>& InMessageHandler) override
+	{
+		if (!InputDevice.IsValid())
+		{
+			TSharedPtr<FOpenXRViveTracker> ViveTrackingInputDevice(new FOpenXRViveTracker(InMessageHandler));
+			InputDevice = ViveTrackingInputDevice;
+
+			return InputDevice;
+		}
+		else
+		{
+			InputDevice.Get()->SetMessageHandler(InMessageHandler);
+			return InputDevice;
+		}
+		return nullptr;
+	}
+
+	virtual TSharedPtr<IInputDevice> GetInputDevice() override
+	{
+		if (!InputDevice.IsValid())
+		{
+			CreateInputDevice(FSlateApplication::Get().GetPlatformApplication()->GetMessageHandler());
+		}
+		return InputDevice;
+	}
+
+private:
+	TSharedPtr<FOpenXRViveTracker> InputDevice;
+};
+
+IMPLEMENT_MODULE(FOpenXRViveTrackerModule, OpenXRViveTracker);
+
+FOpenXRViveTracker::FViveTracker::FViveTracker(XrActionSet InActionSet, FOpenXRPath InRolePath, const char* InName)
+	: ActionSet(InActionSet)
+	, RolePath(InRolePath)
+	, GripAction(XR_NULL_HANDLE)
+	, VibrationAction(XR_NULL_HANDLE)
+	, DeviceId(-1)
+{
+	XrActionCreateInfo Info;
+	Info.type = XR_TYPE_ACTION_CREATE_INFO;
+	Info.next = nullptr;
+	Info.countSubactionPaths = 0;
+	Info.subactionPaths = nullptr;
+
+	FCStringAnsi::Strcpy(Info.localizedActionName, XR_MAX_ACTION_NAME_SIZE, InName);
+	FCStringAnsi::Strcat(Info.localizedActionName, XR_MAX_ACTION_NAME_SIZE, " Pose");
+	FilterActionName(Info.localizedActionName, Info.actionName);
+	Info.actionType = XR_ACTION_TYPE_POSE_INPUT;
+	XR_ENSURE(xrCreateAction(ActionSet, &Info, &GripAction));
+
+	FCStringAnsi::Strcpy(Info.localizedActionName, XR_MAX_ACTION_NAME_SIZE, InName);
+	FCStringAnsi::Strcat(Info.localizedActionName, XR_MAX_ACTION_NAME_SIZE, " Vibration");
+	FilterActionName(Info.localizedActionName, Info.actionName);
+	Info.actionType = XR_ACTION_TYPE_VIBRATION_OUTPUT;
+	XR_ENSURE(xrCreateAction(ActionSet, &Info, &VibrationAction));
+}
+
+void FOpenXRViveTracker::FViveTracker::AddTrackedDevices(FOpenXRHMD* HMD)
+{
+	if (HMD)
+	{
+		DeviceId = HMD->AddTrackedDevice(GripAction, RolePath);
+	}
+}
+
+void FOpenXRViveTracker::FViveTracker::GetSuggestedBindings(TArray<XrActionSuggestedBinding>& OutSuggestedBindings)
+{
+	OutSuggestedBindings.Add(XrActionSuggestedBinding{ GripAction, RolePath / FString("input/grip/pose") });
+	OutSuggestedBindings.Add(XrActionSuggestedBinding{ GripAction, RolePath / FString("output/haptic") });
+}
+
+FOpenXRViveTracker::FOpenXRViveTracker(const TSharedRef<FGenericApplicationMessageHandler>& InMessageHandler)
+	: MessageHandler(InMessageHandler)
+	, DeviceIndex(0)
+	, TrackerActionSet(XR_NULL_HANDLE)
+	, Trackers()
+{
+	// Register modular feature manually
+	IModularFeatures::Get().RegisterModularFeature(IMotionController::GetModularFeatureName(), static_cast<IMotionController*>(this));
+	IModularFeatures::Get().RegisterModularFeature(IOpenXRExtensionPlugin::GetModularFeatureName(), static_cast<IOpenXRExtensionPlugin*>(this));
+	AddKeys();
+
+	// We're implicitly requiring that the OpenXRPlugin has been loaded and
+	// initialized at this point.
+	if (!IOpenXRHMDModule::Get().IsAvailable())
+	{
+		UE_LOG(LogOpenXRViveTracker, Error, TEXT("Error - OpenXRHMDPlugin isn't available"));
+	}
+
+	MotionSourceToEControllerHandMap.Add(TEXT("Pad"), EControllerHand::Pad);
+	MotionSourceToEControllerHandMap.Add(TEXT("ExternalCamera"), EControllerHand::ExternalCamera);
+	MotionSourceToEControllerHandMap.Add(TEXT("Gun"), EControllerHand::Gun);
+	MotionSourceToEControllerHandMap.Add(TEXT("Chest"), EControllerHand::Chest);
+	MotionSourceToEControllerHandMap.Add(TEXT("LeftShoulder"), EControllerHand::LeftShoulder);
+	MotionSourceToEControllerHandMap.Add(TEXT("RightShoulder"), EControllerHand::RightShoulder);
+	MotionSourceToEControllerHandMap.Add(TEXT("LeftElbow"), EControllerHand::LeftElbow);
+	MotionSourceToEControllerHandMap.Add(TEXT("RightElbow"), EControllerHand::RightElbow);
+	MotionSourceToEControllerHandMap.Add(TEXT("Waist"), EControllerHand::Waist);
+	MotionSourceToEControllerHandMap.Add(TEXT("LeftKnee"), EControllerHand::LeftKnee);
+	MotionSourceToEControllerHandMap.Add(TEXT("RightKnee"), EControllerHand::RightKnee);
+	MotionSourceToEControllerHandMap.Add(TEXT("LeftFoot"), EControllerHand::LeftFoot);
+	MotionSourceToEControllerHandMap.Add(TEXT("RightFoot"), EControllerHand::RightFoot);
+
+	// Map the legacy hand enum values that openxr supports
+	MotionSourceToEControllerHandMap.Add(TEXT("EControllerHand::Pad"), EControllerHand::Pad);
+	MotionSourceToEControllerHandMap.Add(TEXT("EControllerHand::ExternalCamera"), EControllerHand::ExternalCamera);
+	MotionSourceToEControllerHandMap.Add(TEXT("EControllerHand::Gun"), EControllerHand::Gun);
+	MotionSourceToEControllerHandMap.Add(TEXT("EControllerHand::Chest"), EControllerHand::Chest);
+	MotionSourceToEControllerHandMap.Add(TEXT("EControllerHand::LeftShoulder"), EControllerHand::LeftShoulder);
+	MotionSourceToEControllerHandMap.Add(TEXT("EControllerHand::RightShoulder"), EControllerHand::RightShoulder);
+	MotionSourceToEControllerHandMap.Add(TEXT("EControllerHand::LeftElbow"), EControllerHand::LeftElbow);
+	MotionSourceToEControllerHandMap.Add(TEXT("EControllerHand::RightElbow"), EControllerHand::RightElbow);
+	MotionSourceToEControllerHandMap.Add(TEXT("EControllerHand::Waist"), EControllerHand::Waist);
+	MotionSourceToEControllerHandMap.Add(TEXT("EControllerHand::LeftKnee"), EControllerHand::LeftKnee);
+	MotionSourceToEControllerHandMap.Add(TEXT("EControllerHand::RightKnee"), EControllerHand::RightKnee);
+	MotionSourceToEControllerHandMap.Add(TEXT("EControllerHand::LeftFoot"), EControllerHand::LeftFoot);
+	MotionSourceToEControllerHandMap.Add(TEXT("EControllerHand::RightFoot"), EControllerHand::RightFoot);
+}
+
+FOpenXRViveTracker::~FOpenXRViveTracker()
+{
+	// Unregister modular feature manually
+	IModularFeatures::Get().UnregisterModularFeature(IMotionController::GetModularFeatureName(), static_cast<IMotionController*>(this));
+	IModularFeatures::Get().UnregisterModularFeature(IOpenXRExtensionPlugin::GetModularFeatureName(), static_cast<IOpenXRExtensionPlugin*>(this));
+}
+
+bool FOpenXRViveTracker::GetRequiredExtensions(TArray<const ANSICHAR*>& OutExtensions)
+{
+	OutExtensions.Add(XR_HTCX_VIVE_TRACKER_INTERACTION_EXTENSION_NAME);
+	return true;
+}
+
+const void* FOpenXRViveTracker::OnGetSystem(XrInstance InInstance, const void* InNext)
+{
+	// Store extension open xr calls to member function pointers for convenient use.
+	XR_ENSURE(xrGetInstanceProcAddr(InInstance, "xrEnumerateViveTrackerPathsHTCX", (PFN_xrVoidFunction*)&xrEnumerateViveTrackerPathsHTCX));
+
+	return InNext;
+}
+
+const void* FOpenXRViveTracker::OnCreateSession(XrInstance InInstance, XrSystemId InSystem, const void* InNext)
+{
+	static FName SystemName(TEXT("OpenXR"));
+	if (GEngine->XRSystem.IsValid() && (GEngine->XRSystem->GetSystemName() == SystemName))
+	{
+		OpenXRHMD = (FOpenXRHMD*)GEngine->XRSystem.Get();
+	}
+
+	if (TrackerActionSet)
+	{
+		XR_ENSURE(xrDestroyActionSet(TrackerActionSet));
+	}
+
+	XrActionSetCreateInfo CreateInfo = { XR_TYPE_ACTION_SET_CREATE_INFO };
+	FCStringAnsi::Strcpy(CreateInfo.actionSetName, XR_MAX_ACTION_SET_NAME_SIZE, "trackers");
+	FCStringAnsi::Strcpy(CreateInfo.localizedActionSetName,
+		XR_MAX_LOCALIZED_ACTION_SET_NAME_SIZE, "Vive Trackers");
+	XR_ENSURE(xrCreateActionSet(InInstance, &CreateInfo, &TrackerActionSet));
+
+	Trackers.Add(EControllerHand::Pad, FViveTracker(TrackerActionSet, FOpenXRPath("/user/vive_tracker_htcx/role/keyboard"), "Keyboard"));
+	Trackers.Add(EControllerHand::ExternalCamera, FViveTracker(TrackerActionSet, FOpenXRPath("/user/vive_tracker_htcx/role/camera"), "External Camera"));
+	Trackers.Add(EControllerHand::Gun, FViveTracker(TrackerActionSet, FOpenXRPath("/user/vive_tracker_htcx/role/handheld_object"), "Handheld Object"));
+	Trackers.Add(EControllerHand::Chest, FViveTracker(TrackerActionSet, FOpenXRPath("/user/vive_tracker_htcx/role/chest"), "Chest"));
+	Trackers.Add(EControllerHand::LeftShoulder, FViveTracker(TrackerActionSet, FOpenXRPath("/user/vive_tracker_htcx/role/left_shoulder"), "Left Shoulder"));
+	Trackers.Add(EControllerHand::RightShoulder, FViveTracker(TrackerActionSet, FOpenXRPath("/user/vive_tracker_htcx/role/right_shoulder"), "Right Shoulder"));
+	Trackers.Add(EControllerHand::LeftElbow, FViveTracker(TrackerActionSet, FOpenXRPath("/user/vive_tracker_htcx/role/left_elbow"), "Left Elbow"));
+	Trackers.Add(EControllerHand::RightElbow, FViveTracker(TrackerActionSet, FOpenXRPath("/user/vive_tracker_htcx/role/right_elbow"), "Right Elbow"));
+	Trackers.Add(EControllerHand::Waist, FViveTracker(TrackerActionSet, FOpenXRPath("/user/vive_tracker_htcx/role/waist"), "Waist"));
+	Trackers.Add(EControllerHand::LeftKnee, FViveTracker(TrackerActionSet, FOpenXRPath("/user/vive_tracker_htcx/role/left_knee"), "Left Knee"));
+	Trackers.Add(EControllerHand::RightKnee, FViveTracker(TrackerActionSet, FOpenXRPath("/user/vive_tracker_htcx/role/right_knee"), "Right Knee"));
+	Trackers.Add(EControllerHand::LeftFoot, FViveTracker(TrackerActionSet, FOpenXRPath("/user/vive_tracker_htcx/role/left_foot"), "Left Foot"));
+	Trackers.Add(EControllerHand::RightFoot, FViveTracker(TrackerActionSet, FOpenXRPath("/user/vive_tracker_htcx/role/right_foot"), "Right Foot"));
+
+	TArray<XrActionSuggestedBinding> Bindings;
+	for (TPair<EControllerHand, FViveTracker> Tracker : Trackers)
+	{
+		Tracker.Value.GetSuggestedBindings(Bindings);
+	}
+
+	XrInteractionProfileSuggestedBinding InteractionProfile = { XR_TYPE_INTERACTION_PROFILE_SUGGESTED_BINDING };
+	InteractionProfile.interactionProfile = FOpenXRPath("/interaction_profiles/htc/vive_tracker_htcx");
+	InteractionProfile.countSuggestedBindings = Bindings.Num();
+	InteractionProfile.suggestedBindings = Bindings.GetData();
+	XR_ENSURE(xrSuggestInteractionProfileBindings(InInstance, &InteractionProfile));
+
+	return InNext;
+}
+
+void FOpenXRViveTracker::OnDestroySession(XrSession InSession)
+{
+	bActionsAttached = false;
+}
+
+void FOpenXRViveTracker::AttachActionSets(TSet<XrActionSet>& OutActionSets)
+{
+	for (TPair<EControllerHand, FViveTracker>& Tracker : Trackers)
+	{
+		Tracker.Value.AddTrackedDevices(OpenXRHMD);
+	}
+
+	OutActionSets.Add(TrackerActionSet);
+
+	bActionsAttached = true;
+}
+
+void FOpenXRViveTracker::GetActiveActionSetsForSync(TArray<XrActiveActionSet>& OutActiveSets)
+{
+	XrActiveActionSet ActiveTrackerSet{ TrackerActionSet, XR_NULL_PATH };
+	OutActiveSets.Add(ActiveTrackerSet);
+}
+
+bool FOpenXRViveTracker::GetControllerOrientationAndPosition(const int32 ControllerIndex, const FName MotionSource, FRotator& OutOrientation, FVector& OutPosition, float WorldToMetersScale) const
+{
+	if (!bActionsAttached || OpenXRHMD == nullptr)
+	{
+		return false;
+	}
+
+	if (ControllerIndex == DeviceIndex && MotionSourceToEControllerHandMap.Contains(MotionSource))
+	{
+		FQuat Orientation;
+		bool Success = OpenXRHMD->GetCurrentPose(GetDeviceIDForMotionSource(MotionSource), Orientation, OutPosition);
+		OutOrientation = FRotator(Orientation);
+		return Success;
+	}
+	return false;
+}
+
+bool FOpenXRViveTracker::GetControllerOrientationAndPositionForTime(const int32 ControllerIndex, const FName MotionSource, FTimespan Time, bool& OutTimeWasUsed, FRotator& OutOrientation, FVector& OutPosition, bool& OutbProvidedLinearVelocity, FVector& OutLinearVelocity, bool& OutbProvidedAngularVelocity, FVector& OutAngularVelocityRadPerSec, bool& OutbProvidedLinearAcceleration, FVector& OutLinearAcceleration, float WorldToMetersScale) const
+{
+	if (ControllerIndex == DeviceIndex && MotionSourceToEControllerHandMap.Contains(MotionSource))
+	{
+		FQuat Orientation;
+		bool Success = OpenXRHMD->GetPoseForTime(GetDeviceIDForMotionSource(MotionSource), Time, OutTimeWasUsed, Orientation, OutPosition, OutbProvidedLinearVelocity, OutLinearVelocity, OutbProvidedAngularVelocity, OutAngularVelocityRadPerSec, OutbProvidedLinearAcceleration, OutLinearAcceleration, WorldToMetersScale);
+		OutOrientation = FRotator(Orientation);
+		return Success;
+	}
+	return false;
+}
+
+ETrackingStatus FOpenXRViveTracker::GetControllerTrackingStatus(const int32 ControllerIndex, const EControllerHand DeviceHand) const
+{
+	if (!bActionsAttached || OpenXRHMD == nullptr)
+	{
+		return ETrackingStatus::NotTracked;
+	}
+
+	XrSession Session = OpenXRHMD->GetSession();
+	if (Session == XR_NULL_HANDLE)
+	{
+		return ETrackingStatus::NotTracked;
+	}
+
+	const FViveTracker* Tracker = Trackers.Find(DeviceHand);
+	if (!Tracker)
+	{
+		return ETrackingStatus::NotTracked;
+	}
+
+	XrActionStateGetInfo Info = { XR_TYPE_ACTION_STATE_GET_INFO };
+	Info.action = Tracker->GripAction;
+	Info.subactionPath = XR_NULL_PATH;
+
+	XrActionStatePose State = { XR_TYPE_ACTION_STATE_POSE };
+	if (!XR_ENSURE(xrGetActionStatePose(Session, &Info, &State)))
+	{
+		return ETrackingStatus::NotTracked;
+	}
+	return State.isActive ? ETrackingStatus::Tracked : ETrackingStatus::NotTracked;
+}
+
+FName FOpenXRViveTracker::GetMotionControllerDeviceTypeName() const
+{
+	const static FName DefaultName(TEXT("OpenXRViveTracker"));
+	return DefaultName;
+}
+
+void FOpenXRViveTracker::EnumerateSources(TArray<FMotionControllerSource>& SourcesOut) const
+{
+	check(IsInGameThread());
+
+	TArray<FName> Sources;
+	MotionSourceToEControllerHandMap.GenerateKeyArray(Sources);
+	SourcesOut.Append(Sources);
+}
+
+void FOpenXRViveTracker::Tick(float DeltaTime)
+{
+}
+
+void FOpenXRViveTracker::SendControllerEvents()
+{
+}
+
+void FOpenXRViveTracker::SetMessageHandler(const TSharedRef<FGenericApplicationMessageHandler>& InMessageHandler)
+{
+	MessageHandler = InMessageHandler;
+}
+
+bool FOpenXRViveTracker::Exec(UWorld* InWorld, const TCHAR* Cmd, FOutputDevice& Ar)
+{
+	return false;
+}
+
+void FOpenXRViveTracker::SetChannelValue(int32 ControllerId, FForceFeedbackChannelType ChannelType, float Value)
+{
+	// Large channel type maps to amplitude. We are interested in amplitude.
+	if ((ChannelType == FForceFeedbackChannelType::LEFT_LARGE) ||
+		(ChannelType == FForceFeedbackChannelType::RIGHT_LARGE))
+	{
+		FHapticFeedbackValues Values(XR_FREQUENCY_UNSPECIFIED, Value);
+		SetHapticFeedbackValues(ControllerId, ChannelType == FForceFeedbackChannelType::LEFT_LARGE ? (int32)EControllerHand::Left : (int32)EControllerHand::Right, Values);
+	}
+}
+
+void FOpenXRViveTracker::SetChannelValues(int32 ControllerId, const FForceFeedbackValues& values)
+{
+	FHapticFeedbackValues leftHaptics = FHapticFeedbackValues(
+		values.LeftSmall,		// frequency
+		values.LeftLarge);		// amplitude
+	FHapticFeedbackValues rightHaptics = FHapticFeedbackValues(
+		values.RightSmall,		// frequency
+		values.RightLarge);		// amplitude
+
+	SetHapticFeedbackValues(
+		ControllerId,
+		(int32)EControllerHand::Left,
+		leftHaptics);
+
+	SetHapticFeedbackValues(
+		ControllerId,
+		(int32)EControllerHand::Right,
+		rightHaptics);
+}
+
+bool FOpenXRViveTracker::IsGamepadAttached() const
+{
+	return false;
+}
+
+// TODO: Refactor API to change the Hand type to EControllerHand
+void FOpenXRViveTracker::SetHapticFeedbackValues(int32 ControllerId, int32 Hand, const FHapticFeedbackValues& Values)
+{
+	if (!bActionsAttached || OpenXRHMD == nullptr)
+	{
+		return;
+	}
+
+	XrSession Session = OpenXRHMD->GetSession();
+
+	if (Session == XR_NULL_HANDLE)
+	{
+		return;
+	}
+
+	if (!OpenXRHMD->IsFocused())
+	{
+		return;
+	}
+
+	XrHapticVibration HapticValue;
+	HapticValue.type = XR_TYPE_HAPTIC_VIBRATION;
+	HapticValue.next = nullptr;
+	HapticValue.duration = XR_INFINITE_DURATION;
+	HapticValue.frequency = Values.Frequency;
+	HapticValue.amplitude = Values.Amplitude;
+
+	FViveTracker* Tracker = Trackers.Find((EControllerHand)Hand);
+	if (Tracker && ControllerId == DeviceIndex)
+	{
+		XrHapticActionInfo HapticActionInfo;
+		HapticActionInfo.type = XR_TYPE_HAPTIC_ACTION_INFO;
+		HapticActionInfo.next = nullptr;
+		HapticActionInfo.subactionPath = XR_NULL_PATH;
+		HapticActionInfo.action = Tracker->VibrationAction;
+		if (Values.HapticBuffer == nullptr && (Values.Amplitude <= 0.0f || Values.Frequency < XR_FREQUENCY_UNSPECIFIED))
+		{
+			XR_ENSURE(xrStopHapticFeedback(Session, &HapticActionInfo));
+		}
+		else
+		{
+			FOpenXRExtensionChainStructPtrs ScopedExtensionChainStructs;
+			if (Values.HapticBuffer != nullptr)
+			{
+				OpenXRHMD->GetApplyHapticFeedbackAddChainStructsDelegate().Broadcast(&HapticValue, ScopedExtensionChainStructs, Values.HapticBuffer);
+			}
+			XR_ENSURE(xrApplyHapticFeedback(Session, &HapticActionInfo, (const XrHapticBaseHeader*)&HapticValue));
+		}
+	}
+}
+
+void FOpenXRViveTracker::GetHapticFrequencyRange(float& MinFrequency, float& MaxFrequency) const
+{
+	MinFrequency = XR_FREQUENCY_UNSPECIFIED;
+	MaxFrequency = XR_FREQUENCY_UNSPECIFIED;
+}
+
+float FOpenXRViveTracker::GetHapticAmplitudeScale() const
+{
+	return 1.0f;
+}
+
+void FOpenXRViveTracker::AddKeys()
+{
+}
+
+XrAction FOpenXRViveTracker::GetActionForMotionSource(FName MotionSource) const
+{
+	const EControllerHand* Hand = MotionSourceToEControllerHandMap.Find(MotionSource);
+	if (Hand == nullptr)
+	{
+		return XR_NULL_HANDLE;
+	}
+	const FViveTracker& Tracker = Trackers[*Hand];
+	return Tracker.GripAction;
+}
+
+int32 FOpenXRViveTracker::GetDeviceIDForMotionSource(FName MotionSource) const
+{
+	const FViveTracker& Tracker = Trackers[MotionSourceToEControllerHandMap.FindChecked(MotionSource)];
+	return Tracker.DeviceId;
+}
+
+XrPath FOpenXRViveTracker::GetRolePathForMotionSource(FName MotionSource) const
+{
+	const FViveTracker& Tracker = Trackers[MotionSourceToEControllerHandMap.FindChecked(MotionSource)];
+	return Tracker.RolePath;
+}
+
+#undef LOCTEXT_NAMESPACE
diff --git a/Engine/Plugins/Runtime/OpenXRViveTracker/Source/OpenXRViveTracker/Private/OpenXRViveTracker.h b/Engine/Plugins/Runtime/OpenXRViveTracker/Source/OpenXRViveTracker/Private/OpenXRViveTracker.h
new file mode 100644
--- /dev/null
+++ b/Engine/Plugins/Runtime/OpenXRViveTracker/Source/OpenXRViveTracker/Private/OpenXRViveTracker.h
@@ -0,0 +1,102 @@
+// Copyright Epic Games, Inc. All Rights Reserved.
+
+#pragma once
+
+#include "CoreMinimal.h"
+#include "GenericPlatform/IInputInterface.h"
+#include "HeadMountedDisplayTypes.h"
+#include "XRMotionControllerBase.h"
+#include "InputCoreTypes.h"
+#include "IInputDevice.h"
+#include "IHapticDevice.h"
+
+#include "IOpenXRExtensionPlugin.h"
+#include "OpenXRCore.h"
+#include "OpenXRHMD.h"
+
+/**
+  * OpenXR ViveTracker
+  */
+class FOpenXRViveTracker :
+	public IOpenXRExtensionPlugin,
+	public IInputDevice,
+	public FXRMotionControllerBase,
+	public IHapticDevice
+{
+public:
+	struct FViveTracker
+	{
+		XrActionSet		ActionSet;
+		FOpenXRPath		RolePath;
+		XrAction		GripAction;
+		XrAction		VibrationAction;
+		int32			DeviceId;
+
+		FViveTracker(XrActionSet InActionSet, FOpenXRPath InRolePath, const char* InName);
+
+		void AddTrackedDevices(class FOpenXRHMD* HMD);
+		void GetSuggestedBindings(TArray<XrActionSuggestedBinding>& OutSuggestedBindings);
+	};
+
+public:
+	FOpenXRViveTracker(const TSharedRef<FGenericApplicationMessageHandler>& InMessageHandler);
+	virtual ~FOpenXRViveTracker();
+
+	/** IOpenXRExtensionPlugin */
+	virtual FString GetDisplayName() override
+	{
+		return FString(TEXT("OpenXRViveTracker"));
+	}
+	virtual bool GetRequiredExtensions(TArray<const ANSICHAR*>& OutExtensions) override;
+	virtual const void* OnGetSystem(XrInstance InInstance, const void* InNext) override;
+	virtual const void* OnCreateSession(XrInstance InInstance, XrSystemId InSystem, const void* InNext) override;
+	virtual void OnDestroySession(XrSession InSession) override;
+	virtual void AttachActionSets(TSet<XrActionSet>& OutActionSets) override;
+	virtual void GetActiveActionSetsForSync(TArray<XrActiveActionSet>& OutActiveSets) override;
+
+	/** IMotionController interface */
+	virtual bool GetControllerOrientationAndPosition(const int32 ControllerIndex, const FName MotionSource, FRotator& OutOrientation, FVector& OutPosition, float WorldToMetersScale) const override;
+	virtual bool GetControllerOrientationAndPosition(const int32 ControllerIndex, const EControllerHand DeviceHand, FRotator& OutOrientation, FVector& OutPosition, float WorldToMetersScale) const override { check(false); return false; }
+	virtual bool GetControllerOrientationAndPositionForTime(const int32 ControllerIndex, const FName MotionSource, FTimespan Time, bool& OutTimeWasUsed, FRotator& OutOrientation, FVector& OutPosition, bool& OutbProvidedLinearVelocity, FVector& OutLinearVelocity, bool& OutbProvidedAngularVelocity, FVector& OutAngularVelocityRadPerSec, bool& OutbProvidedLinearAcceleration, FVector& OutLinearAcceleration, float WorldToMetersScale) const override;
+	virtual ETrackingStatus GetControllerTrackingStatus(const int32 ControllerIndex, const EControllerHand DeviceHand) const override;
+	virtual FName GetMotionControllerDeviceTypeName() const override;
+	virtual void EnumerateSources(TArray<FMotionControllerSource>& SourcesOut) const override;
+
+	/** IInputDevice interface */
+	virtual void Tick(float DeltaTime) override;
+	virtual void SendControllerEvents() override;
+	virtual void SetMessageHandler(const TSharedRef<FGenericApplicationMessageHandler>& InMessageHandler) override;
+	virtual bool Exec(UWorld* InWorld, const TCHAR* Cmd, FOutputDevice& Ar) override;
+	virtual void SetChannelValue(int32 ControllerId, FForceFeedbackChannelType ChannelType, float Value) override;
+	virtual void SetChannelValues(int32 ControllerId, const FForceFeedbackValues &values) override;
+	virtual bool SupportsForceFeedback(int32 ControllerId) override { return true; }
+	virtual bool IsGamepadAttached() const override;
+
+	// IHapticDevice overrides
+	IHapticDevice* GetHapticDevice() override { return (IHapticDevice*)this; }
+	virtual void SetHapticFeedbackValues(int32 ControllerId, int32 Hand, const FHapticFeedbackValues& Values) override;
+
+	virtual void GetHapticFrequencyRange(float& MinFrequency, float& MaxFrequency) const override;
+	virtual float GetHapticAmplitudeScale() const override;
+
+private:
+	void AddKeys();
+
+	XrAction GetActionForMotionSource(FName MotionSource) const;
+	int32 GetDeviceIDForMotionSource(FName MotionSource) const;
+	XrPath GetRolePathForMotionSource(FName MotionSource) const;
+
+	bool bActionsAttached = false;
+
+	PFN_xrEnumerateViveTrackerPathsHTCX xrEnumerateViveTrackerPathsHTCX = nullptr;
+
+	class FOpenXRHMD* OpenXRHMD = nullptr;
+
+	TSharedPtr<FGenericApplicationMessageHandler> MessageHandler;
+	int32 DeviceIndex;
+
+	XrActionSet TrackerActionSet;
+	TMap<EControllerHand, FViveTracker> Trackers;
+
+	TMap<FName, EControllerHand> MotionSourceToEControllerHandMap;
+};
diff --git a/Engine/Plugins/Runtime/OpenXRViveTracker/Source/OpenXRViveTracker/Public/IOpenXRViveTrackerModule.h b/Engine/Plugins/Runtime/OpenXRViveTracker/Source/OpenXRViveTracker/Public/IOpenXRViveTrackerModule.h
new file mode 100644
--- /dev/null
+++ b/Engine/Plugins/Runtime/OpenXRViveTracker/Source/OpenXRViveTracker/Public/IOpenXRViveTrackerModule.h
@@ -0,0 +1,48 @@
+// Copyright Epic Games, Inc. All Rights Reserved.
+
+#pragma once
+
+#include "CoreMinimal.h"
+#include "Modules/ModuleManager.h"
+#include "IInputDeviceModule.h"
+
+enum class EHandKeypoint : uint8;
+enum class EControllerHand : uint8;
+
+/**
+ * The public interface to this module.  In most cases, this interface is only public to sibling modules
+ * within this plugin.
+ */
+class IOpenXRViveTrackerModule : public IInputDeviceModule
+{
+public:
+	/**
+	* Singleton-like access to this module's interface.  This is just for convenience!
+	* Beware of calling this during the shutdown phase, though. Your module might have been
+	* unloaded already.
+	*
+	* @return Returns singleton instance, loading the module on demand if needed
+	*/
+	static inline IOpenXRViveTrackerModule& Get()
+	{
+		return FModuleManager::LoadModuleChecked<IOpenXRViveTrackerModule>("OpenXRViveTracker");
+	}
+
+	/**
+	* Checks to see if this module is loaded and ready.  It is only valid to call Get() if
+	* IsAvailable() returns true.
+	*
+	* @return True if the module is loaded and ready to use
+	*/
+	static inline bool IsAvailable()
+	{
+		return FModuleManager::Get().IsModuleLoaded("OpenXRViveTracker");
+	}
+
+	/**
+	 * Returns the Input Device associated with this InputDeviceModule.
+	 *
+	 * @return Shared pointer to the Input Device.
+	 */
+	virtual TSharedPtr<IInputDevice> GetInputDevice() = 0;
+};
