Remove legacy motion sources from the OpenXR Vive Tracker plugin

diff --git a/Engine/Plugins/Runtime/OpenXRViveTracker/Source/OpenXRViveTracker/Private/OpenXRViveTracker.cpp b/Engine/Plugins/Runtime/OpenXRViveTracker/Source/OpenXRViveTracker/Private/OpenXRViveTracker.cpp
--- a/Engine/Plugins/Runtime/OpenXRViveTracker/Source/OpenXRViveTracker/Private/OpenXRViveTracker.cpp
+++ b/Engine/Plugins/Runtime/OpenXRViveTracker/Source/OpenXRViveTracker/Private/OpenXRViveTracker.cpp
@@ -156,21 +156,6 @@ FOpenXRViveTracker::FOpenXRViveTracker(const TSharedRef<FGenericApplicationMessa
 	MotionSourceToEControllerHandMap.Add(TEXT("RightKnee"), EControllerHand::RightKnee);
 	MotionSourceToEControllerHandMap.Add(TEXT("LeftFoot"), EControllerHand::LeftFoot);
 	MotionSourceToEControllerHandMap.Add(TEXT("RightFoot"), EControllerHand::RightFoot);
-
-	// Map the legacy hand enum values that openxr supports
-	MotionSourceToEControllerHandMap.Add(TEXT("EControllerHand::Pad"), EControllerHand::Pad);
-	MotionSourceToEControllerHandMap.Add(TEXT("EControllerHand::ExternalCamera"), EControllerHand::ExternalCamera);
-	MotionSourceToEControllerHandMap.Add(TEXT("EControllerHand::Gun"), EControllerHand::Gun);
-	MotionSourceToEControllerHandMap.Add(TEXT("EControllerHand::Chest"), EControllerHand::Chest);
-	MotionSourceToEControllerHandMap.Add(TEXT("EControllerHand::LeftShoulder"), EControllerHand::LeftShoulder);
-	MotionSourceToEControllerHandMap.Add(TEXT("EControllerHand::RightShoulder"), EControllerHand::RightShoulder);
-	MotionSourceToEControllerHandMap.Add(TEXT("EControllerHand::LeftElbow"), EControllerHand::LeftElbow);
-	MotionSourceToEControllerHandMap.Add(TEXT("EControllerHand::RightElbow"), EControllerHand::RightElbow);
-	MotionSourceToEControllerHandMap.Add(TEXT("EControllerHand::Waist"), EControllerHand::Waist);
-	MotionSourceToEControllerHandMap.Add(TEXT("EControllerHand::LeftKnee"), EControllerHand::LeftKnee);
-	MotionSourceToEControllerHandMap.Add(TEXT("EControllerHand::RightKnee"), EControllerHand::RightKnee);
-	MotionSourceToEControllerHandMap.Add(TEXT("EControllerHand::LeftFoot"), EControllerHand::LeftFoot);
-	MotionSourceToEControllerHandMap.Add(TEXT("EControllerHand::RightFoot"), EControllerHand::RightFoot);
 }
 
 FOpenXRViveTracker::~FOpenXRViveTracker()
