Fix OpenXRViveTracker crashes when no HMD connected.

diff --git a/Engine/Plugins/Runtime/OpenXRViveTracker/Source/OpenXRViveTracker/Private/OpenXRViveTracker.cpp b/Engine/Plugins/Runtime/OpenXRViveTracker/Source/OpenXRViveTracker/Private/OpenXRViveTracker.cpp
--- a/Engine/Plugins/Runtime/OpenXRViveTracker/Source/OpenXRViveTracker/Private/OpenXRViveTracker.cpp
+++ b/Engine/Plugins/Runtime/OpenXRViveTracker/Source/OpenXRViveTracker/Private/OpenXRViveTracker.cpp
@@ -277,7 +277,7 @@ bool FOpenXRViveTracker::GetControllerOrientationAndPosition(const int32 Control
 
 bool FOpenXRViveTracker::GetControllerOrientationAndPositionForTime(const int32 ControllerIndex, const FName MotionSource, FTimespan Time, bool& OutTimeWasUsed, FRotator& OutOrientation, FVector& OutPosition, bool& OutbProvidedLinearVelocity, FVector& OutLinearVelocity, bool& OutbProvidedAngularVelocity, FVector& OutAngularVelocityAsAxisAndLength, bool& OutbProvidedLinearAcceleration, FVector& OutLinearAcceleration, float WorldToMetersScale) const
 {
-	if (ControllerIndex == DeviceIndex && MotionSourceToEControllerHandMap.Contains(MotionSource))
+	if (OpenXRHMD && ControllerIndex == DeviceIndex && MotionSourceToEControllerHandMap.Contains(MotionSource))
 	{
 		FQuat Orientation;
 		bool Success = OpenXRHMD->GetPoseForTime(GetDeviceIDForMotionSource(MotionSource), Time, OutTimeWasUsed, Orientation, OutPosition, OutbProvidedLinearVelocity, OutLinearVelocity, OutbProvidedAngularVelocity, OutAngularVelocityAsAxisAndLength, OutbProvidedLinearAcceleration, OutLinearAcceleration, WorldToMetersScale);
