Fix incorrect haptic action mapping for Vive trackers

diff --git a/Engine/Plugins/Runtime/OpenXRViveTracker/Source/OpenXRViveTracker/Private/OpenXRViveTracker.cpp b/Engine/Plugins/Runtime/OpenXRViveTracker/Source/OpenXRViveTracker/Private/OpenXRViveTracker.cpp
--- a/Engine/Plugins/Runtime/OpenXRViveTracker/Source/OpenXRViveTracker/Private/OpenXRViveTracker.cpp
+++ b/Engine/Plugins/Runtime/OpenXRViveTracker/Source/OpenXRViveTracker/Private/OpenXRViveTracker.cpp
@@ -122,7 +122,7 @@ void FOpenXRViveTracker::FViveTracker::AddTrackedDevices(FOpenXRHMD* HMD)
 void FOpenXRViveTracker::FViveTracker::GetSuggestedBindings(TArray<XrActionSuggestedBinding>& OutSuggestedBindings)
 {
 	OutSuggestedBindings.Add(XrActionSuggestedBinding{ GripAction, RolePath / FString("input/grip/pose") });
-	OutSuggestedBindings.Add(XrActionSuggestedBinding{ GripAction, RolePath / FString("output/haptic") });
+	OutSuggestedBindings.Add(XrActionSuggestedBinding{ VibrationAction, RolePath / FString("output/haptic") });
 }
 
 FOpenXRViveTracker::FOpenXRViveTracker(const TSharedRef<FGenericApplicationMessageHandler>& InMessageHandler)
