Enumerate unassigned trackers

diff --git a/Engine/Plugins/Runtime/OpenXRViveTracker/Source/OpenXRViveTracker/Private/OpenXRViveTracker.cpp b/Engine/Plugins/Runtime/OpenXRViveTracker/Source/OpenXRViveTracker/Private/OpenXRViveTracker.cpp
--- a/Engine/Plugins/Runtime/OpenXRViveTracker/Source/OpenXRViveTracker/Private/OpenXRViveTracker.cpp
+++ b/Engine/Plugins/Runtime/OpenXRViveTracker/Source/OpenXRViveTracker/Private/OpenXRViveTracker.cpp
@@ -171,12 +171,10 @@ bool FOpenXRViveTracker::GetRequiredExtensions(TArray<const ANSICHAR*>& OutExten
 	return true;
 }
 
-const void* FOpenXRViveTracker::OnGetSystem(XrInstance InInstance, const void* InNext)
+void FOpenXRViveTracker::PostCreateInstance(XrInstance InInstance)
 {
 	// Store extension open xr calls to member function pointers for convenient use.
 	XR_ENSURE(xrGetInstanceProcAddr(InInstance, "xrEnumerateViveTrackerPathsHTCX", (PFN_xrVoidFunction*)&xrEnumerateViveTrackerPathsHTCX));
-
-	return InNext;
 }
 
 const void* FOpenXRViveTracker::OnCreateSession(XrInstance InInstance, XrSystemId InSystem, const void* InNext)
@@ -218,6 +216,27 @@ const void* FOpenXRViveTracker::OnCreateSession(XrInstance InInstance, XrSystemI
 		Tracker.Value.GetSuggestedBindings(Bindings);
 	}
 
+	uint32_t PathCount = 0;
+	TArray<XrViveTrackerPathsHTCX> TrackerPaths;
+	XR_ENSURE(xrEnumerateViveTrackerPathsHTCX(InInstance, 0, &PathCount, nullptr));
+	TrackerPaths.SetNum(PathCount);
+	XR_ENSURE(xrEnumerateViveTrackerPathsHTCX(InInstance, PathCount, &PathCount, TrackerPaths.GetData()));
+	check(TrackerPaths.Num() == PathCount);
+
+	// Add LiveLink poses for trackers that have no role assigned
+	uint32 TrackerIndex = 0;
+	for (const XrViveTrackerPathsHTCX& TrackerPath : TrackerPaths)
+	{
+		if (TrackerPath.rolePath == XR_NULL_PATH)
+		{
+			char Name[XR_MAX_ACTION_NAME_SIZE];
+			FCStringAnsi::Snprintf(Name, XR_MAX_ACTION_NAME_SIZE, "Tracker %u", TrackerIndex++);
+			FViveTracker Tracker(TrackerActionSet, TrackerPath.persistentPath, Name);
+			UnassignedTrackers.Add(Tracker);
+			Tracker.GetSuggestedBindings(Bindings);
+		}
+	}
+
 	XrInteractionProfileSuggestedBinding InteractionProfile = { XR_TYPE_INTERACTION_PROFILE_SUGGESTED_BINDING };
 	InteractionProfile.interactionProfile = FOpenXRPath("/interaction_profiles/htc/vive_tracker_htcx");
 	InteractionProfile.countSuggestedBindings = Bindings.Num();
@@ -229,6 +248,8 @@ const void* FOpenXRViveTracker::OnCreateSession(XrInstance InInstance, XrSystemI
 
 void FOpenXRViveTracker::OnDestroySession(XrSession InSession)
 {
+	Trackers.Reset();
+	UnassignedTrackers.Reset();
 	bActionsAttached = false;
 }
 
@@ -239,6 +260,11 @@ void FOpenXRViveTracker::AttachActionSets(TSet<XrActionSet>& OutActionSets)
 		Tracker.Value.AddTrackedDevices(OpenXRHMD);
 	}
 
+	for (FViveTracker& Tracker : UnassignedTrackers)
+	{
+		Tracker.AddTrackedDevices(OpenXRHMD);
+	}
+
 	OutActionSets.Add(TrackerActionSet);
 
 	bActionsAttached = true;
diff --git a/Engine/Plugins/Runtime/OpenXRViveTracker/Source/OpenXRViveTracker/Private/OpenXRViveTracker.h b/Engine/Plugins/Runtime/OpenXRViveTracker/Source/OpenXRViveTracker/Private/OpenXRViveTracker.h
--- a/Engine/Plugins/Runtime/OpenXRViveTracker/Source/OpenXRViveTracker/Private/OpenXRViveTracker.h
+++ b/Engine/Plugins/Runtime/OpenXRViveTracker/Source/OpenXRViveTracker/Private/OpenXRViveTracker.h
@@ -48,7 +48,7 @@ public:
 		return FString(TEXT("OpenXRViveTracker"));
 	}
 	virtual bool GetRequiredExtensions(TArray<const ANSICHAR*>& OutExtensions) override;
-	virtual const void* OnGetSystem(XrInstance InInstance, const void* InNext) override;
+	virtual void PostCreateInstance(XrInstance InInstance) override;
 	virtual const void* OnCreateSession(XrInstance InInstance, XrSystemId InSystem, const void* InNext) override;
 	virtual void OnDestroySession(XrSession InSession) override;
 	virtual void AttachActionSets(TSet<XrActionSet>& OutActionSets) override;
@@ -97,6 +97,7 @@ private:
 
 	XrActionSet TrackerActionSet;
 	TMap<EControllerHand, FViveTracker> Trackers;
+	TArray<FViveTracker> UnassignedTrackers;
 
 	TMap<FName, EControllerHand> MotionSourceToEControllerHandMap;
 };
